Add test for ChaosUserDataPT ClearData function

diff --git a/Engine/Source/Programs/ChaosUserDataPTTests/Private/ChaosUserDataPTTests.cpp b/Engine/Source/Programs/ChaosUserDataPTTests/Private/ChaosUserDataPTTests.cpp
--- a/Engine/Source/Programs/ChaosUserDataPTTests/Private/ChaosUserDataPTTests.cpp
+++ b/Engine/Source/Programs/ChaosUserDataPTTests/Private/ChaosUserDataPTTests.cpp
@@ -25,11 +25,21 @@ TEST_CASE("ChaosUserDataPT", "[integration]")
 	// Make a box geometry
 	auto BoxGeom = TSharedPtr<Chaos::FImplicitObject, ESPMode::ThreadSafe>(new Chaos::TBox<Chaos::FReal, 3>(Chaos::FVec3(-1, -1, -1), Chaos::FVec3(1, 1, 1)));
 
-	// Add a proxy to the solver, advance twice to get a valid internal handle
+	// Add some proxies to the solver
 	Chaos::FSingleParticlePhysicsProxy* Proxy0 = Chaos::FSingleParticlePhysicsProxy::Create(Chaos::FGeometryParticle::CreateParticle());
+	Chaos::FSingleParticlePhysicsProxy* Proxy1 = Chaos::FSingleParticlePhysicsProxy::Create(Chaos::FGeometryParticle::CreateParticle());
+	Chaos::FSingleParticlePhysicsProxy* Proxy2 = Chaos::FSingleParticlePhysicsProxy::Create(Chaos::FGeometryParticle::CreateParticle());
 	Chaos::FRigidBodyHandle_External& HandleExternal0 = Proxy0->GetGameThreadAPI();
+	Chaos::FRigidBodyHandle_External& HandleExternal1 = Proxy1->GetGameThreadAPI();
+	Chaos::FRigidBodyHandle_External& HandleExternal2 = Proxy2->GetGameThreadAPI();
 	HandleExternal0.SetGeometry(BoxGeom);
+	HandleExternal1.SetGeometry(BoxGeom);
+	HandleExternal2.SetGeometry(BoxGeom);
 	Solver->RegisterObject(Proxy0);
+	Solver->RegisterObject(Proxy1);
+	Solver->RegisterObject(Proxy2);
+
+	// Advance the solver twice to make sure the PT handles are created and in the evolution
 	Solver->AdvanceAndDispatch_External(DeltaTime)->Wait();
 	Solver->AdvanceAndDispatch_External(DeltaTime)->Wait();
 
@@ -119,4 +129,38 @@ TEST_CASE("ChaosUserDataPT", "[integration]")
 		});
 		Solver->AdvanceAndDispatch_External(DeltaTime)->Wait();
 	}
+
+	SECTION("Clearing all data from a userdata manager")
+	{
+		// Add data to three particles, propagate it to the PT
+		TestUserData->SetData_GT(HandleExternal0, TestString);
+		TestUserData->SetData_GT(HandleExternal1, TestString);
+		TestUserData->SetData_GT(HandleExternal2, TestString);
+		Solver->AdvanceAndDispatch_External(DeltaTime)->Wait();
+		Solver->AdvanceAndDispatch_External(DeltaTime)->Wait();
+
+		// Clear all data from the userdata manager, but after that set data back on
+		// particle 2 - the fact that it happened _after_ clearing should mean it is
+		// still there for particle 2 after the clear reaches the PT.
+		TestUserData->ClearData_GT();
+		TestUserData->SetData_GT(HandleExternal2, TestString);
+
+		// Check to see that the data is still there at first
+		Solver->EnqueueCommandImmediate([&]()
+		{
+			REQUIRE(*TestUserData->GetData_PT(*Proxy0->GetPhysicsThreadAPI()) == TestString);
+			REQUIRE(*TestUserData->GetData_PT(*Proxy1->GetPhysicsThreadAPI()) == TestString);
+			REQUIRE(*TestUserData->GetData_PT(*Proxy2->GetPhysicsThreadAPI()) == TestString);
+		});
+		Solver->AdvanceAndDispatch_External(DeltaTime)->Wait();
+
+		// Check make sure that after a couple updates, the data is cleared
+		Solver->EnqueueCommandImmediate([&]()
+		{
+			REQUIRE(TestUserData->GetData_PT(*Proxy0->GetPhysicsThreadAPI()) == nullptr);
+			REQUIRE(TestUserData->GetData_PT(*Proxy1->GetPhysicsThreadAPI()) == nullptr);
+			REQUIRE(*TestUserData->GetData_PT(*Proxy2->GetPhysicsThreadAPI()) == TestString);
+		});
+		Solver->AdvanceAndDispatch_External(DeltaTime)->Wait();
+	}
 }
