Add a test to ChaosUserDataPTTests which changes data and ensures that that the new value makes it from GT to PT.

diff --git a/Engine/Source/Programs/ChaosUserDataPTTests/Private/ChaosUserDataPTTests.cpp b/Engine/Source/Programs/ChaosUserDataPTTests/Private/ChaosUserDataPTTests.cpp
--- a/Engine/Source/Programs/ChaosUserDataPTTests/Private/ChaosUserDataPTTests.cpp
+++ b/Engine/Source/Programs/ChaosUserDataPTTests/Private/ChaosUserDataPTTests.cpp
@@ -13,7 +13,8 @@ class FTestUserData : public Chaos::TUserDataManagerPT<FString> { };
 TEST_CASE("ChaosUserDataPT", "[integration]")
 {
 	const float DeltaTime = 1.f;
-	const FString TestString = TEXT("TestData");
+	const FString TestString1 = TEXT("TestData1");
+	const FString TestString2 = TEXT("TestData2");
 
 	// Create a solver in the solvers module
 	FChaosSolversModule* Module = FChaosSolversModule::GetModule();
@@ -47,7 +48,7 @@ TEST_CASE("ChaosUserDataPT", "[integration]")
 	SECTION("Data propagates from GT to PT")
 	{
 		// Add userdata to the particle
-		TestUserData->SetData_GT(HandleExternal0, TestString);
+		TestUserData->SetData_GT(HandleExternal0, TestString1);
 
 		// The first callback should show no data because the ensure will occur before
 		// the sim callback has occurred.
@@ -60,7 +61,26 @@ TEST_CASE("ChaosUserDataPT", "[integration]")
 		// The data should make it to the physics thread by this point
 		Solver->EnqueueCommandImmediate([&]()
 		{
-			REQUIRE(*TestUserData->GetData_PT(*Proxy0->GetPhysicsThreadAPI()) == TestString);
+			REQUIRE(*TestUserData->GetData_PT(*Proxy0->GetPhysicsThreadAPI()) == TestString1);
+		});
+		Solver->AdvanceAndDispatch_External(DeltaTime)->Wait();
+	}
+
+	// Change data
+	SECTION("Data updates propagate from GT to PT")
+	{
+		// Add userdata to the particle
+		TestUserData->SetData_GT(HandleExternal0, TestString1);
+		Solver->AdvanceAndDispatch_External(DeltaTime)->Wait();
+
+		// Set the userdata to something else
+		TestUserData->SetData_GT(HandleExternal0, TestString2);
+		Solver->AdvanceAndDispatch_External(DeltaTime)->Wait();
+
+		// The data should make it to the physics thread by this point
+		Solver->EnqueueCommandImmediate([&]()
+		{
+			REQUIRE(*TestUserData->GetData_PT(*Proxy0->GetPhysicsThreadAPI()) == TestString2);
 		});
 		Solver->AdvanceAndDispatch_External(DeltaTime)->Wait();
 	}
@@ -69,7 +89,7 @@ TEST_CASE("ChaosUserDataPT", "[integration]")
 	SECTION("Data removals propagate from GT to PT")
 	{
 		// Add userdata to the particle and advance it to the physics thread
-		TestUserData->SetData_GT(HandleExternal0, TestString);
+		TestUserData->SetData_GT(HandleExternal0, TestString1);
 		Solver->AdvanceAndDispatch_External(DeltaTime)->Wait();
 
 		// Delete the data
@@ -78,7 +98,7 @@ TEST_CASE("ChaosUserDataPT", "[integration]")
 		// Data should exist for one more update
 		Solver->EnqueueCommandImmediate([&]()
 		{
-			REQUIRE(*TestUserData->GetData_PT(*Proxy0->GetPhysicsThreadAPI()) == TestString);
+			REQUIRE(*TestUserData->GetData_PT(*Proxy0->GetPhysicsThreadAPI()) == TestString1);
 		});
 		Solver->AdvanceAndDispatch_External(DeltaTime)->Wait();
 
@@ -108,7 +128,7 @@ TEST_CASE("ChaosUserDataPT", "[integration]")
 	{
 		// Add data to a particle, make sure it gets to PT, then delete the particle.
 		const Chaos::FUniqueIdx UniqueIdx0 = HandleExternal0.UniqueIdx();
-		TestUserData->SetData_GT(HandleExternal0, TestString);
+		TestUserData->SetData_GT(HandleExternal0, TestString1);
 		Solver->AdvanceAndDispatch_External(DeltaTime)->Wait();
 		Solver->UnregisterObject(Proxy0);
 		Solver->AdvanceAndDispatch_External(DeltaTime)->Wait();
@@ -133,9 +153,9 @@ TEST_CASE("ChaosUserDataPT", "[integration]")
 	SECTION("Clearing all data from a userdata manager")
 	{
 		// Add data to three particles, propagate it to the PT
-		TestUserData->SetData_GT(HandleExternal0, TestString);
-		TestUserData->SetData_GT(HandleExternal1, TestString);
-		TestUserData->SetData_GT(HandleExternal2, TestString);
+		TestUserData->SetData_GT(HandleExternal0, TestString1);
+		TestUserData->SetData_GT(HandleExternal1, TestString1);
+		TestUserData->SetData_GT(HandleExternal2, TestString1);
 		Solver->AdvanceAndDispatch_External(DeltaTime)->Wait();
 		Solver->AdvanceAndDispatch_External(DeltaTime)->Wait();
 
@@ -143,14 +163,14 @@ TEST_CASE("ChaosUserDataPT", "[integration]")
 		// particle 2 - the fact that it happened _after_ clearing should mean it is
 		// still there for particle 2 after the clear reaches the PT.
 		TestUserData->ClearData_GT();
-		TestUserData->SetData_GT(HandleExternal2, TestString);
+		TestUserData->SetData_GT(HandleExternal2, TestString1);
 
 		// Check to see that the data is still there at first
 		Solver->EnqueueCommandImmediate([&]()
 		{
-			REQUIRE(*TestUserData->GetData_PT(*Proxy0->GetPhysicsThreadAPI()) == TestString);
-			REQUIRE(*TestUserData->GetData_PT(*Proxy1->GetPhysicsThreadAPI()) == TestString);
-			REQUIRE(*TestUserData->GetData_PT(*Proxy2->GetPhysicsThreadAPI()) == TestString);
+			REQUIRE(*TestUserData->GetData_PT(*Proxy0->GetPhysicsThreadAPI()) == TestString1);
+			REQUIRE(*TestUserData->GetData_PT(*Proxy1->GetPhysicsThreadAPI()) == TestString1);
+			REQUIRE(*TestUserData->GetData_PT(*Proxy2->GetPhysicsThreadAPI()) == TestString1);
 		});
 		Solver->AdvanceAndDispatch_External(DeltaTime)->Wait();
 
@@ -159,7 +179,7 @@ TEST_CASE("ChaosUserDataPT", "[integration]")
 		{
 			REQUIRE(TestUserData->GetData_PT(*Proxy0->GetPhysicsThreadAPI()) == nullptr);
 			REQUIRE(TestUserData->GetData_PT(*Proxy1->GetPhysicsThreadAPI()) == nullptr);
-			REQUIRE(*TestUserData->GetData_PT(*Proxy2->GetPhysicsThreadAPI()) == TestString);
+			REQUIRE(*TestUserData->GetData_PT(*Proxy2->GetPhysicsThreadAPI()) == TestString1);
 		});
 		Solver->AdvanceAndDispatch_External(DeltaTime)->Wait();
 	}
