Fixed user data tests crashing.
- Test was expecting valid event reference to be returned but this cannot be assumed, instead we just wait on all pending tasks inside the solver
- Test was failing to destroy the solver and potentially running past the end of the test scope with events still pending in the queue, on executing this could crash as they reference a function-local string which will have been destroyed

diff --git a/Engine/Source/Programs/ChaosUserDataPTTests/Private/ChaosUserDataPTTests.cpp b/Engine/Source/Programs/ChaosUserDataPTTests/Private/ChaosUserDataPTTests.cpp
--- a/Engine/Source/Programs/ChaosUserDataPTTests/Private/ChaosUserDataPTTests.cpp
+++ b/Engine/Source/Programs/ChaosUserDataPTTests/Private/ChaosUserDataPTTests.cpp
@@ -9,6 +9,16 @@
 
 class FTestUserData : public Chaos::TUserDataManagerPT<FString> { };
 
+// Advance a solver once with the provided DeltaTime, then wait for any async tasks to finish before continuing
+void AdvanceAndWait(Chaos::FPBDRigidsSolver* InSolver, float DeltaTime)
+{
+	if(InSolver)
+	{
+		InSolver->AdvanceAndDispatch_External(DeltaTime);
+		InSolver->WaitOnPendingTasks_External();
+	}
+}
+
 TEST_CASE("ChaosUserDataPT", "[integration]")
 {
 	const float DeltaTime = 1.f;
@@ -40,8 +50,8 @@ TEST_CASE("ChaosUserDataPT", "[integration]")
 	Solver->RegisterObject(Proxy2);
 
 	// Advance the solver twice to make sure the PT handles are created and in the evolution
-	Solver->AdvanceAndDispatch_External(DeltaTime)->Wait();
-	Solver->AdvanceAndDispatch_External(DeltaTime)->Wait();
+	AdvanceAndWait(Solver, DeltaTime);
+	AdvanceAndWait(Solver, DeltaTime);
 
 	// Add data
 	SECTION("Data propagates from GT to PT")
@@ -55,14 +65,14 @@ TEST_CASE("ChaosUserDataPT", "[integration]")
 		{
 			REQUIRE(TestUserData->GetData_PT(*Proxy0->GetPhysicsThreadAPI()) == nullptr);
 		});
-		Solver->AdvanceAndDispatch_External(DeltaTime)->Wait();
+		AdvanceAndWait(Solver, DeltaTime);
 
 		// The data should make it to the physics thread by this point
 		Solver->EnqueueCommandImmediate([&]()
 		{
 			REQUIRE(*TestUserData->GetData_PT(*Proxy0->GetPhysicsThreadAPI()) == TestString1);
 		});
-		Solver->AdvanceAndDispatch_External(DeltaTime)->Wait();
+		AdvanceAndWait(Solver, DeltaTime);
 	}
 
 	// Change data
@@ -70,18 +80,18 @@ TEST_CASE("ChaosUserDataPT", "[integration]")
 	{
 		// Add userdata to the particle
 		TestUserData->SetData_GT(HandleExternal0, TestString1);
-		Solver->AdvanceAndDispatch_External(DeltaTime)->Wait();
+		AdvanceAndWait(Solver, DeltaTime);
 
 		// Set the userdata to something else
 		TestUserData->SetData_GT(HandleExternal0, TestString2);
-		Solver->AdvanceAndDispatch_External(DeltaTime)->Wait();
+		AdvanceAndWait(Solver, DeltaTime);
 
 		// The data should make it to the physics thread by this point
 		Solver->EnqueueCommandImmediate([&]()
 		{
 			REQUIRE(*TestUserData->GetData_PT(*Proxy0->GetPhysicsThreadAPI()) == TestString2);
 		});
-		Solver->AdvanceAndDispatch_External(DeltaTime)->Wait();
+		AdvanceAndWait(Solver, DeltaTime);
 	}
 
 	// Delete data
@@ -89,7 +99,7 @@ TEST_CASE("ChaosUserDataPT", "[integration]")
 	{
 		// Add userdata to the particle and advance it to the physics thread
 		TestUserData->SetData_GT(HandleExternal0, TestString1);
-		Solver->AdvanceAndDispatch_External(DeltaTime)->Wait();
+		AdvanceAndWait(Solver, DeltaTime);
 
 		// Delete the data
 		TestUserData->RemoveData_GT(HandleExternal0);
@@ -99,14 +109,14 @@ TEST_CASE("ChaosUserDataPT", "[integration]")
 		{
 			REQUIRE(*TestUserData->GetData_PT(*Proxy0->GetPhysicsThreadAPI()) == TestString1);
 		});
-		Solver->AdvanceAndDispatch_External(DeltaTime)->Wait();
+		AdvanceAndWait(Solver, DeltaTime);
 
 		// Data should be deleted at this point
 		Solver->EnqueueCommandImmediate([&]()
 		{
 			REQUIRE(TestUserData->GetData_PT(*Proxy0->GetPhysicsThreadAPI()) == nullptr);
 		});
-		Solver->AdvanceAndDispatch_External(DeltaTime)->Wait();
+		AdvanceAndWait(Solver, DeltaTime);
 	}
 
 	// Delete data from particle that doesn't have it
@@ -114,12 +124,12 @@ TEST_CASE("ChaosUserDataPT", "[integration]")
 	{
 		// Delete data that isn't there
 		TestUserData->RemoveData_GT(HandleExternal0);
-		Solver->AdvanceAndDispatch_External(DeltaTime)->Wait();
+		AdvanceAndWait(Solver, DeltaTime);
 		Solver->EnqueueCommandImmediate([&]()
 		{
 			REQUIRE(TestUserData->GetData_PT(*Proxy0->GetPhysicsThreadAPI()) == nullptr);
 		});
-		Solver->AdvanceAndDispatch_External(DeltaTime)->Wait();
+		AdvanceAndWait(Solver, DeltaTime);
 	}
 
 	// Make sure a particle with a recycled index can't get a deleted particle's userdata
@@ -128,10 +138,10 @@ TEST_CASE("ChaosUserDataPT", "[integration]")
 		// Add data to a particle, make sure it gets to PT, then delete the particle.
 		const Chaos::FUniqueIdx UniqueIdx0 = HandleExternal0.UniqueIdx();
 		TestUserData->SetData_GT(HandleExternal0, TestString1);
-		Solver->AdvanceAndDispatch_External(DeltaTime)->Wait();
+		AdvanceAndWait(Solver, DeltaTime);
 		Solver->UnregisterObject(Proxy0);
-		Solver->AdvanceAndDispatch_External(DeltaTime)->Wait();
-		Solver->AdvanceAndDispatch_External(DeltaTime)->Wait();
+		AdvanceAndWait(Solver, DeltaTime);
+		AdvanceAndWait(Solver, DeltaTime);
 
 		struct FMockHandle
 		{
@@ -146,7 +156,7 @@ TEST_CASE("ChaosUserDataPT", "[integration]")
 			const FMockHandle MockHandle0 = FMockHandle(UniqueIdx0);
 			REQUIRE(TestUserData->GetData_PT(MockHandle0) == nullptr);
 		});
-		Solver->AdvanceAndDispatch_External(DeltaTime)->Wait();
+		AdvanceAndWait(Solver, DeltaTime);
 	}
 
 	SECTION("Clearing all data from a userdata manager")
@@ -155,8 +165,8 @@ TEST_CASE("ChaosUserDataPT", "[integration]")
 		TestUserData->SetData_GT(HandleExternal0, TestString1);
 		TestUserData->SetData_GT(HandleExternal1, TestString1);
 		TestUserData->SetData_GT(HandleExternal2, TestString1);
-		Solver->AdvanceAndDispatch_External(DeltaTime)->Wait();
-		Solver->AdvanceAndDispatch_External(DeltaTime)->Wait();
+		AdvanceAndWait(Solver, DeltaTime);
+		AdvanceAndWait(Solver, DeltaTime);
 
 		// Clear all data from the userdata manager, but after that set data back on
 		// particle 2 - the fact that it happened _after_ clearing should mean it is
@@ -171,7 +181,7 @@ TEST_CASE("ChaosUserDataPT", "[integration]")
 			REQUIRE(*TestUserData->GetData_PT(*Proxy1->GetPhysicsThreadAPI()) == TestString1);
 			REQUIRE(*TestUserData->GetData_PT(*Proxy2->GetPhysicsThreadAPI()) == TestString1);
 		});
-		Solver->AdvanceAndDispatch_External(DeltaTime)->Wait();
+		AdvanceAndWait(Solver, DeltaTime);
 
 		// Check make sure that after a couple updates, the data is cleared
 		Solver->EnqueueCommandImmediate([&]()
@@ -180,6 +190,10 @@ TEST_CASE("ChaosUserDataPT", "[integration]")
 			REQUIRE(TestUserData->GetData_PT(*Proxy1->GetPhysicsThreadAPI()) == nullptr);
 			REQUIRE(*TestUserData->GetData_PT(*Proxy2->GetPhysicsThreadAPI()) == TestString1);
 		});
-		Solver->AdvanceAndDispatch_External(DeltaTime)->Wait();
+		AdvanceAndWait(Solver, DeltaTime);
 	}
+
+	Solver->WaitOnPendingTasks_External();
+	Module->DestroySolver(Solver);
 }
+
